Make it possible to distinguish between controller types when using WinDualshock

diff --git a/Engine/Plugins/Runtime/Windows/WinDualShock/Source/WinDualShock/Private/WinDualShock.cpp b/Engine/Plugins/Runtime/Windows/WinDualShock/Source/WinDualShock/Private/WinDualShock.cpp
--- a/Engine/Plugins/Runtime/Windows/WinDualShock/Source/WinDualShock/Private/WinDualShock.cpp
+++ b/Engine/Plugins/Runtime/Windows/WinDualShock/Source/WinDualShock/Private/WinDualShock.cpp
@@ -20,7 +20,6 @@ DEFINE_LOG_CATEGORY_STATIC(LogWinDualShock, Log, All);
 #include LIBSCEPAD_PLATFORM_INCLUDE
 
 static FName InputClassName = FName("FWinDualShock");
-static FString InputControllerIdentifier = TEXT("DualShock4");
 
 class FWinDualShock : public IInputDevice
 {
@@ -82,17 +81,17 @@ public:
 	{
 		for (int32 UserIndex = 0; UserIndex < SCE_USER_SERVICE_MAX_LOGIN_USERS; UserIndex++)
 		{
-			FInputDeviceScope InputScope(nullptr, InputClassName, UserIndex, InputControllerIdentifier);
+			FInputDeviceScope InputScope(this, InputClassName, UserIndex, Controllers.GetControllerTypeIdentifier(UserIndex));
 			Controllers.SendControllerEvents(UserIndex, MessageHandler);
 		}
 	}
 
-	void SetChannelValue (int32 ControllerId, FForceFeedbackChannelType ChannelType, float Value)
+	void SetChannelValue(int32 ControllerId, FForceFeedbackChannelType ChannelType, float Value)
 	{
 		Controllers.SetForceFeedbackChannelValue(ControllerId, ChannelType, Value);
 	}
 
-	void SetChannelValues (int32 ControllerId, const FForceFeedbackValues &Values)
+	void SetChannelValues(int32 ControllerId, const FForceFeedbackValues &Values)
 	{
 		Controllers.SetForceFeedbackChannelValues(ControllerId, Values);
 	}
