Fix non-unity

diff --git a/Engine/Plugins/Runtime/Windows/WinDualShock/Source/WinDualShock/Private/WinDualShock.h b/Engine/Plugins/Runtime/Windows/WinDualShock/Source/WinDualShock/Private/WinDualShock.h
--- a/Engine/Plugins/Runtime/Windows/WinDualShock/Source/WinDualShock/Private/WinDualShock.h
+++ b/Engine/Plugins/Runtime/Windows/WinDualShock/Source/WinDualShock/Private/WinDualShock.h
@@ -6,6 +6,14 @@
 
 enum ESonyControllerType : unsigned int;
 
+namespace Audio
+{
+	/**
+	 * Typed identifier for Audio Device Id
+	 */
+	using FDeviceId = uint32;
+}
+
 /**
  * Supported endpoint types
  */
