Prevent multiple endpoint submix instances from overlapping on the same hardware device, which isn't supported

diff --git a/Engine/Plugins/Runtime/Windows/WinDualShock/Source/WinDualShock/Private/WinDualShock.h b/Engine/Plugins/Runtime/Windows/WinDualShock/Source/WinDualShock/Private/WinDualShock.h
--- a/Engine/Plugins/Runtime/Windows/WinDualShock/Source/WinDualShock/Private/WinDualShock.h
+++ b/Engine/Plugins/Runtime/Windows/WinDualShock/Source/WinDualShock/Private/WinDualShock.h
@@ -20,7 +20,8 @@ namespace Audio
 enum EWinDualShockPortType
 {
 	PadSpeakers,
-	Vibration	// ignored on DualShock controllers
+	Vibration,	// ignored on DualShock controllers
+	WinDualShockPortType_Count
 };
 
 /**
@@ -73,18 +74,21 @@ public:
 
 	/**
 	 * Endpoint reference counting to know when device is no longer needed
+	 * Returns true only for the first endpoint of a port type, to avoid multiple audio sends interleaving
 	 */
-	void AddEndpoint()
+	bool AddEndpoint(EWinDualShockPortType PortType)
 	{
-		EndpointCount++;
+		int PreviousCount = EndpointCount[PortType]++;
+		return PreviousCount == 0;
 	}
-	void RemoveEndpoint()
+	void RemoveEndpoint(EWinDualShockPortType PortType)
 	{
-		EndpointCount--;
+		EndpointCount[PortType]--;
 	}
 	int GetEndpointCount() const
 	{
-		return EndpointCount;
+		static_assert(WinDualShockPortType_Count == 2);
+		return EndpointCount[0] + EndpointCount[1];
 	}
 
 	FDeviceKey GetDeviceKey() const
@@ -94,7 +98,7 @@ public:
 
 protected:
 	FDeviceKey			DeviceKey;
-	std::atomic<int>	EndpointCount{ 0 };
+	std::atomic<int>	EndpointCount[WinDualShockPortType_Count]{ 0 };
 };
 
 
diff --git a/Engine/Plugins/Runtime/Windows/WinDualShock/Source/WinDualShock/Private/WinDualShockExternalEndpoints.h b/Engine/Plugins/Runtime/Windows/WinDualShock/Source/WinDualShock/Private/WinDualShockExternalEndpoints.h
--- a/Engine/Plugins/Runtime/Windows/WinDualShock/Source/WinDualShock/Private/WinDualShockExternalEndpoints.h
+++ b/Engine/Plugins/Runtime/Windows/WinDualShock/Source/WinDualShock/Private/WinDualShockExternalEndpoints.h
@@ -14,16 +14,22 @@ class FExternalWinDualShockEndpoint : public IAudioEndpoint
 {
 private:
 	TSharedRef<IWinDualShockAudioDevice> Device;
+	bool bAllowedEndpoint;
 
 public:
-	FExternalWinDualShockEndpoint(TSharedRef<IWinDualShockAudioDevice> InDevice) : Device(InDevice)
+	FExternalWinDualShockEndpoint(TSharedRef<IWinDualShockAudioDevice> InDevice) 
+		: Device(InDevice), bAllowedEndpoint(InDevice->AddEndpoint(PortType))
 	{
-		Device->AddEndpoint();
 	}
 
 	virtual ~FExternalWinDualShockEndpoint()
 	{
-		Device->RemoveEndpoint();
+		Device->RemoveEndpoint(PortType);
+	}
+
+	bool IsEndpointAllowed() const
+	{
+		return bAllowedEndpoint;
 	}
 
 protected:
@@ -63,7 +69,10 @@ protected:
 	{
 		check(NumChannels == GetNumChannels());
 
-		Device->PushAudio(PortType, InAudio, NumChannels);
+		if (bAllowedEndpoint)
+		{
+			Device->PushAudio(PortType, InAudio, NumChannels);
+		}
 		return true;
 	}
 
