Warn if multiple endpoint submix have been registered on a particular endpoint, as that isn't supported.

diff --git a/Engine/Plugins/Runtime/Windows/WinDualShock/Source/WinDualShock/Private/WinDualShock.cpp b/Engine/Plugins/Runtime/Windows/WinDualShock/Source/WinDualShock/Private/WinDualShock.cpp
--- a/Engine/Plugins/Runtime/Windows/WinDualShock/Source/WinDualShock/Private/WinDualShock.cpp
+++ b/Engine/Plugins/Runtime/Windows/WinDualShock/Source/WinDualShock/Private/WinDualShock.cpp
@@ -663,7 +663,12 @@ public:
 		};
 		TSharedRef<IWinDualShockAudioDevice> Device = 
 			InputDevice ? InputDevice->GetAudioDevice(DeviceKey, AudioDeviceCreateFn) : GetEarlyAudioDevice(DeviceKey, AudioDeviceCreateFn);
-		return TUniquePtr<IAudioEndpoint>(new FExternalWinDualShockEndpoint<PortType>(Device));
+		FExternalWinDualShockEndpoint<PortType>* Endpoint = new FExternalWinDualShockEndpoint<PortType>(Device);
+		if (!Endpoint->IsEndpointAllowed())
+		{
+			UE_LOG(LogWinDualShock, Warning, TEXT("Only one instance of %s endpoint submit is allowed.  Additional instances ignored."), *GetEndpointTypeName().ToString());
+		}
+		return TUniquePtr<IAudioEndpoint>(Endpoint);
 	}
 
 	virtual UClass* GetCustomSettingsClass() const override
