Fix for Static Analisys Warnings

diff --git a/Engine/Plugins/Runtime/StateTree/Source/StateTreeModule/Private/Blueprint/StateTreeConditionBlueprintBase.cpp b/Engine/Plugins/Runtime/StateTree/Source/StateTreeModule/Private/Blueprint/StateTreeConditionBlueprintBase.cpp
--- a/Engine/Plugins/Runtime/StateTree/Source/StateTreeModule/Private/Blueprint/StateTreeConditionBlueprintBase.cpp
+++ b/Engine/Plugins/Runtime/StateTree/Source/StateTreeModule/Private/Blueprint/StateTreeConditionBlueprintBase.cpp
@@ -44,14 +44,20 @@ bool FStateTreeBlueprintConditionWrapper::TestCondition(FStateTreeExecutionConte
 FText FStateTreeBlueprintConditionWrapper::GetDescription(const FGuid& ID, FStateTreeDataView InstanceData, const IStateTreeBindingLookup& BindingLookup) const
 {
 	const UClass* Class = Cast<const UClass>(InstanceData.GetStruct());
-	if (UBlueprint* ClassBP = Cast<UBlueprint>(Class->ClassGeneratedBy))
+
+	if (Class != nullptr)
 	{
-		if (!ClassBP->BlueprintDescription.IsEmpty())
+		if (UBlueprint* ClassBP = Cast<UBlueprint>(Class->ClassGeneratedBy))
 		{
-			return FText::FromString(ClassBP->BlueprintDescription);
+			if (!ClassBP->BlueprintDescription.IsEmpty())
+			{
+				return FText::FromString(ClassBP->BlueprintDescription);
+			}
 		}
+
+		return Class->GetDisplayNameText();
 	}
 	
-	return Class != nullptr ? Class->GetDisplayNameText() : FText::GetEmpty();
+	return FText::GetEmpty();
 }
 #endif
diff --git a/Engine/Plugins/Runtime/Windows/WinDualShock/Source/WinDualShock/Private/WinDualShock.cpp b/Engine/Plugins/Runtime/Windows/WinDualShock/Source/WinDualShock/Private/WinDualShock.cpp
--- a/Engine/Plugins/Runtime/Windows/WinDualShock/Source/WinDualShock/Private/WinDualShock.cpp
+++ b/Engine/Plugins/Runtime/Windows/WinDualShock/Source/WinDualShock/Private/WinDualShock.cpp
@@ -235,7 +235,7 @@ public:
 
 	virtual ~FWinDualShockAudioDeviceImpl()
 	{
-		TearDownAudio();
+		TearDownAudioInternal();
 	}
 
 	void SetInputDevice(FWinDualShock* InInputDevice)
@@ -426,26 +426,7 @@ public:
 
 	virtual void TearDownAudio()
 	{
-		if (bAudioIsSetup.load())
-		{
-			bAudioIsSetup.store(false, std::memory_order_release);
-			if (AudioDevice.IsInitialized())
-			{
-				AudioDevice.Shutdown();
-			}
-
-			if (AudioPadSpeakerBuffer)
-			{
-				delete AudioPadSpeakerBuffer;
-				AudioPadSpeakerBuffer = NULL;
-			}
-
-			if (AudioVibrationBuffer)
-			{
-				delete AudioVibrationBuffer;
-				AudioVibrationBuffer = NULL;
-			}
-		}
+		TearDownAudioInternal();
 	}
 
 protected:
@@ -616,6 +597,32 @@ protected:
 	TCircularQueue<float>*	AudioPadSpeakerBuffer;
 	TCircularQueue<float>*	AudioVibrationBuffer;
 	FWinDualShock*			InputDevice = nullptr;
+
+private:
+
+	void TearDownAudioInternal()
+	{
+		if (bAudioIsSetup.load())
+		{
+			bAudioIsSetup.store(false, std::memory_order_release);
+			if (AudioDevice.IsInitialized())
+			{
+				AudioDevice.Shutdown();
+			}
+
+			if (AudioPadSpeakerBuffer)
+			{
+				delete AudioPadSpeakerBuffer;
+				AudioPadSpeakerBuffer = NULL;
+			}
+
+			if (AudioVibrationBuffer)
+			{
+				delete AudioVibrationBuffer;
+				AudioVibrationBuffer = NULL;
+			}
+		}
+	}
 };
 
 template<EWinDualShockPortType PortType>
@@ -707,13 +714,13 @@ private:
 class FWinDualShockPlugin : public IInputDeviceModule
 {
 #if DUALSHOCK4_SUPPORT
+	// audio devices allocated early before FWinDualShock created
+	TMap<FDeviceKey, TSharedRef<FWinDualShockAudioDeviceImpl>>				EarlyDeviceMap;
+
 	// exported endpoints
 	FExternalDualShockEndpointFactory<EWinDualShockPortType::PadSpeakers>	PadSpeakerEndpoint;
 	FExternalDualShockEndpointFactory<EWinDualShockPortType::Vibration>		VibrationEndpoint;
 
-	// audio devices allocated early before FWinDualShock created
-	TMap<FDeviceKey, TSharedRef<FWinDualShockAudioDeviceImpl>>				EarlyDeviceMap;
-
 public:
 	FWinDualShockPlugin()
 		: PadSpeakerEndpoint(EarlyDeviceMap), VibrationEndpoint(EarlyDeviceMap)
diff --git a/Engine/Source/Editor/Kismet/Private/SReplaceNodeReferences.cpp b/Engine/Source/Editor/Kismet/Private/SReplaceNodeReferences.cpp
--- a/Engine/Source/Editor/Kismet/Private/SReplaceNodeReferences.cpp
+++ b/Engine/Source/Editor/Kismet/Private/SReplaceNodeReferences.cpp
@@ -690,7 +690,7 @@ void SReplaceNodeReferences::FindAllReplacementsComplete(TArray<FImaginaryFiBDat
 		}
 	}
 
-	if (HasValidSource() && SelectedTargetReferenceItem.IsValid())
+	if (HasValidSource() && SelectedTargetReferenceItem.IsValid() && SourceProperty!=nullptr )
 	{
 		FMemberReference SourceVariableReference;
 		SourceVariableReference.SetFromField<FProperty>(SourceProperty, SourceProperty->GetOwnerClass());
diff --git a/Engine/Source/Runtime/Engine/Private/Materials/MaterialInstance.cpp b/Engine/Source/Runtime/Engine/Private/Materials/MaterialInstance.cpp
--- a/Engine/Source/Runtime/Engine/Private/Materials/MaterialInstance.cpp
+++ b/Engine/Source/Runtime/Engine/Private/Materials/MaterialInstance.cpp
@@ -2083,7 +2083,7 @@ void UMaterialInstance::CacheShadersForResources(EShaderPlatform ShaderPlatform,
 	UMaterial* BaseMaterial = GetMaterial();
 #if WITH_EDITOR
 	check(!HasAnyFlags(RF_NeedPostLoad));
-	check(!BaseMaterial->HasAnyFlags(RF_NeedPostLoad));
+	check(BaseMaterial!=nullptr && !BaseMaterial->HasAnyFlags(RF_NeedPostLoad));
 #endif
 
 #if WITH_EDITOR
diff --git a/Engine/Source/Runtime/Engine/Private/SkeletalRenderGPUSkin.cpp b/Engine/Source/Runtime/Engine/Private/SkeletalRenderGPUSkin.cpp
--- a/Engine/Source/Runtime/Engine/Private/SkeletalRenderGPUSkin.cpp
+++ b/Engine/Source/Runtime/Engine/Private/SkeletalRenderGPUSkin.cpp
@@ -2089,7 +2089,10 @@ void FDynamicSkelMeshObjectDataGPUSkin::InitDynamicSkelMeshObjectDataGPUSkin(
 #if !(UE_BUILD_SHIPPING || UE_BUILD_TEST)
 	check(!MeshComponentSpaceTransforms.Num());
 	// append instead of equals to avoid alloc
-	MeshComponentSpaceTransforms.Append(InMeshComponent->GetComponentSpaceTransforms());
+	if (InMeshComponent != nullptr)
+	{
+		MeshComponentSpaceTransforms.Append(InMeshComponent->GetComponentSpaceTransforms());
+	}	
 #endif
 	SectionIdsUseByActiveMorphTargets.Empty();
 	// find number of morphs that are currently weighted and will affect the mesh
