Fix stack corruption due to local variable out of scope being accessed by another thread

diff --git a/Engine/Plugins/Runtime/Windows/WinDualShock/Source/WinDualShock/Private/WinDualShock.cpp b/Engine/Plugins/Runtime/Windows/WinDualShock/Source/WinDualShock/Private/WinDualShock.cpp
--- a/Engine/Plugins/Runtime/Windows/WinDualShock/Source/WinDualShock/Private/WinDualShock.cpp
+++ b/Engine/Plugins/Runtime/Windows/WinDualShock/Source/WinDualShock/Private/WinDualShock.cpp
@@ -57,10 +57,6 @@ public:
 		bool bDSTouchAxisButtons = false;
 		bool bDSMouseEvents = false;
 		bool bDSMotionEvents = false;
-		float DSPadSpeakerGain = 1.f;
-		float DSHeadphonesGain = 1.f;
-		float DSMicrophoneGain = 1.f;
-		float DSOutputGain = 1.f;
 
 		if (GConfig)
 		{
@@ -75,32 +71,37 @@ public:
 
 			// Configure SonyControllers to emit motion events if the application wants them
 			GConfig->GetBool(TEXT("SonyController"), TEXT("bDSMotionEvents"), bDSMotionEvents, GEngineIni);
-
-			// Configure SonyControllers gain (between 0.0-1.0, default 1.0)
-			GConfig->GetFloat(TEXT("SonyController"), TEXT("DSPadSpeakerGain"), DSPadSpeakerGain, GEngineIni);
-			GConfig->GetFloat(TEXT("SonyController"), TEXT("DSHeadphonesGain"), DSHeadphonesGain, GEngineIni);
-			GConfig->GetFloat(TEXT("SonyController"), TEXT("DSMicrophoneGain"), DSMicrophoneGain, GEngineIni);
 		}
 
 		// On Windows, we need to compensate the padspeaker and haptics volume
 		// by the platform headroom attenuation
 
-		if (IsInGameThread() && GEngine)
+		FAudioThread::RunCommandOnAudioThread([this]()
 		{
-			if (auto MainAudioDevice = GEngine->GetMainAudioDevice())
+			float DSPadSpeakerGain = 1.f;
+			float DSHeadphonesGain = 1.f;
+			float DSMicrophoneGain = 1.f;
+			float DSOutputGain = 1.f;
+
+			// Configure SonyControllers gain (between 0.0-1.0, default 1.0)
+			GConfig->GetFloat(TEXT("SonyController"), TEXT("DSPadSpeakerGain"), DSPadSpeakerGain, GEngineIni);
+			GConfig->GetFloat(TEXT("SonyController"), TEXT("DSHeadphonesGain"), DSHeadphonesGain, GEngineIni);
+			GConfig->GetFloat(TEXT("SonyController"), TEXT("DSMicrophoneGain"), DSMicrophoneGain, GEngineIni);
+
+			if (GEngine)
 			{
-				FAudioThread::RunCommandOnAudioThread([MainAudioDevice, &DSOutputGain]()
-					{
-						DSOutputGain = 1.0f / FMath::Clamp(MainAudioDevice->GetPlatformAudioHeadroom(), TNumericLimits<float>::Min(), 1.f);
-					});
+				if (auto MainAudioDevice = GEngine->GetMainAudioDevice())
+				{
+					DSOutputGain = 1.0f / FMath::Clamp(MainAudioDevice->GetPlatformAudioHeadroom(), TNumericLimits<float>::Min(), 1.f);
+				}
 			}
-		}
+			Controllers.SetAudioGain(DSPadSpeakerGain, DSHeadphonesGain, DSMicrophoneGain, DSOutputGain);
+		});
 
 		Controllers.SetEmitTouchEvents(bDSTouchEvents);
 		Controllers.SetEmitTouchAxisEvents(bDSTouchAxisButtons);
 		Controllers.SetEmitMouseEvents(bDSMouseEvents);
 		Controllers.SetEmitMotionEvents(bDSMotionEvents);
-		Controllers.SetAudioGain(DSPadSpeakerGain, DSHeadphonesGain, DSMicrophoneGain, DSOutputGain);
 
 		for (int32 UserIndex = 0; UserIndex < SCE_USER_SERVICE_MAX_LOGIN_USERS; UserIndex++)
 		{
diff --git a/Engine/Plugins/Runtime/Windows/WinDualShock/Source/WinDualShock/Private/WinDualShockControllers.h b/Engine/Plugins/Runtime/Windows/WinDualShock/Source/WinDualShock/Private/WinDualShockControllers.h
--- a/Engine/Plugins/Runtime/Windows/WinDualShock/Source/WinDualShock/Private/WinDualShockControllers.h
+++ b/Engine/Plugins/Runtime/Windows/WinDualShock/Source/WinDualShock/Private/WinDualShockControllers.h
@@ -25,6 +25,10 @@ public:
 		HeadphonesGain = InHeadphonesGain;
 		MicrophoneGain = InMicrophoneGain;
 		OutputGain = InOutputGain;
+		for (int32 UserIndex = 0; UserIndex < SCE_USER_SERVICE_MAX_LOGIN_USERS; UserIndex++)
+		{
+			bGainChanged[UserIndex] = true;
+		}
 	}
 
 	float GetOutputGain()
