WinDualShock: Fixed DualSense controller being recognized as DualShock4 during the first update after launching the game.

The problem was caused by FWinDualShock::SendControllerEvents setting up a FInputDeviceScope before the controlller's type is initialized (which happens in SendControllerEvents).

An additional factor contributing to the issue is that while we do initialize controllers (and query their types) in FWinDualShock constructor by calling ConnectStateToUser, the type info is not available at this point. It's because controlllers' state is internally updated by libScePad library on a thread, so even if our controller is connected before we launch the game, it will simply return 'not connected' controller type, which we previously ignored.

Resolved by additionallly updating the controllers' types. The added benefit is that it fixed the problem of GetControllerTypeIdentifier and GetControllerTypeName returning stale type/name for controllers that have been disconnected.

#rn fix PS5 DualSense controllers used on PC are no longer incorrectly recognized as DualShock4 when they are initially connected. Previouslly, disconnected controllers would return DualShock4 as their type. Now, DualSense/DualShock4 types are only returned for controllers which state is 'connected'. Otherwise, None is returned for them.

diff --git a/Engine/Plugins/Runtime/Windows/WinDualShock/Source/WinDualShock/Private/WinDualShock.cpp b/Engine/Plugins/Runtime/Windows/WinDualShock/Source/WinDualShock/Private/WinDualShock.cpp
--- a/Engine/Plugins/Runtime/Windows/WinDualShock/Source/WinDualShock/Private/WinDualShock.cpp
+++ b/Engine/Plugins/Runtime/Windows/WinDualShock/Source/WinDualShock/Private/WinDualShock.cpp
@@ -131,8 +131,20 @@ public:
 	{
 		for (int32 UserIndex = 0; UserIndex < SCE_USER_SERVICE_MAX_LOGIN_USERS; UserIndex++)
 		{
-			FInputDeviceScope InputScope(this, InputClassName, UserIndex, Controllers.GetControllerTypeIdentifierName(UserIndex));
-			Controllers.SendControllerEvents(UserIndex, MessageHandler);
+			// On Windows all controllers exist and can be opened, but until their state changes to connected
+			// (which is updated asynchronously on an internal thread in libScePad) we don't know the controller type.
+			// That's why checking querying the type in ConnectStateToUser is not enough and we need this lazy update.
+			if (Controllers.GetControllerTypeIdentifier(UserIndex) == ESonyControllerType::None)
+			{
+				Controllers.RefreshControllerType(UserIndex);
+			}
+
+			// Don't send events until we have the controller type as we need to set the device scope up.
+			if (Controllers.GetControllerTypeIdentifier(UserIndex) != ESonyControllerType::None)
+			{
+				FInputDeviceScope InputScope(this, InputClassName, UserIndex, Controllers.GetControllerTypeIdentifierName(UserIndex));
+				Controllers.SendControllerEvents(UserIndex, MessageHandler);
+			}
 		}
 		UpdateAudioDevices();
 	}
diff --git a/Engine/Plugins/Runtime/Windows/WinDualShock/Source/WinDualShock/Private/WinDualShockControllers.h b/Engine/Plugins/Runtime/Windows/WinDualShock/Source/WinDualShock/Private/WinDualShockControllers.h
--- a/Engine/Plugins/Runtime/Windows/WinDualShock/Source/WinDualShock/Private/WinDualShockControllers.h
+++ b/Engine/Plugins/Runtime/Windows/WinDualShock/Source/WinDualShock/Private/WinDualShockControllers.h
@@ -41,6 +41,11 @@ public:
 		return bSupportsAudio[UserIndex];
 	}
 
+	void RefreshControllerType(int32 UserIndex)
+	{
+		ControllerTypeIdentifiers[UserIndex] = GetControllerType(UserIndex);
+	}
+
 private:
 	float OutputGain = 1.0f;
 };
