Fixed shadowed variable.

diff --git a/Engine/Plugins/Runtime/Windows/WinDualShock/Source/WinDualShock/Private/WinDualShock.cpp b/Engine/Plugins/Runtime/Windows/WinDualShock/Source/WinDualShock/Private/WinDualShock.cpp
--- a/Engine/Plugins/Runtime/Windows/WinDualShock/Source/WinDualShock/Private/WinDualShock.cpp
+++ b/Engine/Plugins/Runtime/Windows/WinDualShock/Source/WinDualShock/Private/WinDualShock.cpp
@@ -559,12 +559,12 @@ protected:
 			bStreamStarted = false;
 		}
 
-		void SubmitBuffer(FCallbackInfo& CallbackInfo) const
+		void SubmitBuffer(FCallbackInfo& InCallbackInfo) const
 		{
 			XAUDIO2_BUFFER XAudio2Buffer = { 0 };
 			XAudio2Buffer.AudioBytes = EWinDualShockDefaults::NumFrames * NumOutputChannels * sizeof(float);
 			XAudio2Buffer.pAudioData = (const BYTE*)OutputBuffer.GetData();
-			XAudio2Buffer.pContext = &CallbackInfo;
+			XAudio2Buffer.pContext = &InCallbackInfo;
 
 			// Submit buffer to the output streaming voice
 			HRESULT hr;
