Fixed shadowing warnings caused by 'None' in the global namespace.

diff --git a/Engine/Plugins/Runtime/Windows/WinDualShock/Source/WinDualShock/Private/WinDualShock.h b/Engine/Plugins/Runtime/Windows/WinDualShock/Source/WinDualShock/Private/WinDualShock.h
--- a/Engine/Plugins/Runtime/Windows/WinDualShock/Source/WinDualShock/Private/WinDualShock.h
+++ b/Engine/Plugins/Runtime/Windows/WinDualShock/Source/WinDualShock/Private/WinDualShock.h
@@ -4,7 +4,7 @@
 
 #include "CoreMinimal.h"
 
-enum ESonyControllerType : unsigned int;
+enum class ESonyControllerType : unsigned int;
 
 namespace Audio
 {
