WinDualShock disabled on arm64 in the build rules, not UEBuildWindows

diff --git a/Engine/Plugins/Runtime/Windows/WinDualShock/Source/WinDualShock/WinDualShock.Build.cs b/Engine/Plugins/Runtime/Windows/WinDualShock/Source/WinDualShock/WinDualShock.Build.cs
--- a/Engine/Plugins/Runtime/Windows/WinDualShock/Source/WinDualShock/WinDualShock.Build.cs
+++ b/Engine/Plugins/Runtime/Windows/WinDualShock/Source/WinDualShock/WinDualShock.Build.cs
@@ -27,7 +27,7 @@ namespace UnrealBuildTool.Rules
 			System.Type LibScePadType = System.Type.GetType("LibScePad");
 			bool bHasSupport = false;
 			string PlatformName = "";
-			if (LibScePadType != null)
+			if (LibScePadType != null && Target.Architecture != UnrealArch.Arm64)
 			{
 				bHasSupport = (bool)LibScePadType.GetMethod("GetPadLibLocation").Invoke(null, new object[] { EngineDirectory, Target.WindowsPlatform.Compiler, null });
 				PlatformName = (string)LibScePadType.GetMethod("GetPlatformName").Invoke(null, new object[] { EngineDirectory, Target.WindowsPlatform.Compiler });
diff --git a/Engine/Source/Programs/UnrealBuildTool/Platform/Windows/UEBuildWindows.cs b/Engine/Source/Programs/UnrealBuildTool/Platform/Windows/UEBuildWindows.cs
--- a/Engine/Source/Programs/UnrealBuildTool/Platform/Windows/UEBuildWindows.cs
+++ b/Engine/Source/Programs/UnrealBuildTool/Platform/Windows/UEBuildWindows.cs
@@ -1266,9 +1266,6 @@ namespace UnrealBuildTool
 							// DVP
 							"MediaIOFramework",
 							"Composure",
-
-							// sce lib
-							"WinDualShock",
 						});
 					}
 
