Fix DualSense controllers not functioning with Wasapi audio backend enabled

diff --git a/Engine/Plugins/Runtime/Windows/WinDualShock/Source/WinDualShock/Private/WinDualShock.cpp b/Engine/Plugins/Runtime/Windows/WinDualShock/Source/WinDualShock/Private/WinDualShock.cpp
--- a/Engine/Plugins/Runtime/Windows/WinDualShock/Source/WinDualShock/Private/WinDualShock.cpp
+++ b/Engine/Plugins/Runtime/Windows/WinDualShock/Source/WinDualShock/Private/WinDualShock.cpp
@@ -5,6 +5,7 @@
 #include "Engine/Engine.h"
 #include "Misc/CoreDelegates.h"
 #include "Misc/MessageDialog.h"
+#include "Misc/Paths.h"
 #include "Modules/ModuleManager.h"
 #include "GenericPlatform/IInputInterface.h"
 #include "IInputDeviceModule.h"
@@ -21,9 +22,32 @@
 THIRD_PARTY_INCLUDES_START
 #include <Audioclient.h>					// WASAPI api
 #include <winreg.h>
+
+// ******************************************************************************************************
+// Duplicate of definitions found in AudioMixerPlatformXAudio2.h/cpp
+// ******************************************************************************************************
+// Notes on how the headers and libs are used in UE XAudio2 implementation:
+//   Intel Windows:
+//     Use the MS XAudio2.9 redist header/library/DLL to guarantee 2.9 on any Windows version
+//   Arm64 Windows:
+//     Must use the system XAudio2_9.dll. We could use the Win10 SDK's verison of xaudio2.h, however
+//     UE includes a copy of the DirectX SDK headers, which will interfere with finding the SDK version.
+//     The redist header is usable, however, even tho we are not using the redist libraries. We do not 
+//     need to manually load the DLL because linking with the SDK arm64 .lib, without a DELAYLOAD, 
+//     will make the OS automatically load the DLL correctly.s
+//   Consoles:
+//     Similar to the Arm64 above, except the copy of DirectX SDK headers are not used, so consoles
+//     can use the xaudio2.h from the SDK
+#define USE_REDIST_HEADER PLATFORM_WINDOWS
+#define USE_REDIST_LIB PLATFORM_WINDOWS && !PLATFORM_CPU_ARM_FAMILY
+
+#if USE_REDIST_HEADER
 #include <xaudio2redist.h>
+#else // USE_REDIST_HEADER
+#include <xaudio2.h>
+#endif // USE_REDIST_HEADER
 THIRD_PARTY_INCLUDES_END
-#endif
+#endif // PLATFORM_WINDOWS
 
 #include <pad.h>
 #include <pad_audio.h>
@@ -478,6 +502,27 @@ protected:
 	public:
 		bool Initialize(const FString& OutputDeviceId, ESonyControllerType ControllerType)
 		{
+#if USE_REDIST_LIB
+			// Work around the fact the x64 version of XAudio2_7.dll does not properly ref count
+			// by forcing it to be always loaded
+
+			// Load the xaudio2 library and keep a handle so we can free it on teardown
+			// Note: windows internally ref-counts the library per call to load library so 
+			// when we call FreeLibrary, it will only free it once the refcount is zero
+			// Also, FPlatformProcess::GetDllHandle should not be used, as it will not increase ref count further if the library is already loaded.
+			// FPaths::ConvertRelativePathToFull is used for parity with how GetDllHandle calls LoadLibrary.
+			static const FString XAudio2_9Redist = FPaths::EngineDir() / TEXT("Binaries/ThirdParty/Windows/XAudio2_9/x64/xaudio2_9redist.dll");
+			XAudio2Dll = LoadLibrary(*FPaths::ConvertRelativePathToFull(XAudio2_9Redist));
+
+			// returning null means we failed to load XAudio2, which means everything will fail
+			if (XAudio2Dll == nullptr)
+			{
+				UE_LOG(LogInit, Warning, TEXT("Failed to load XAudio2 dll"));
+				FMessageDialog::Open(EAppMsgType::Ok, NSLOCTEXT("Audio", "XAudio2Missing", "XAudio2.7 is not installed. Make sure you have XAudio 2.7 installed. XAudio 2.7 is available in the DirectX End-User Runtime (June 2010)."));
+				return false;
+			}
+#endif // USE_REDIST_LIB
+
 			HRESULT hr;
 			hr = XAudio2Create(&Context, 0, XAUDIO2_DEFAULT_PROCESSOR);
 			if (FAILED(hr))
@@ -554,6 +599,17 @@ protected:
 				MasterVoice = nullptr;
 			}
 			SAFE_RELEASE(Context);
+
+			if (XAudio2Dll)
+			{
+				if (!FreeLibrary(XAudio2Dll))
+				{
+					UE_LOG(LogAudioMixer, Warning, TEXT("FXAudio2Device failed to free XAudio2 Dll"));
+				}
+
+				XAudio2Dll = nullptr;
+			}
+
 			bStreamStarted = false;
 		}
 
@@ -596,6 +652,7 @@ protected:
 		uint32					NumOutputChannels;
 
 	private:
+		HMODULE					XAudio2Dll = nullptr;
 		IXAudio2*				Context = nullptr;
 		IXAudio2MasteringVoice*	MasterVoice = nullptr;
 		IXAudio2SourceVoice*	SourceVoice = nullptr;
