Add 'ExtraRootPath' to WinDualShock.Build.cs

#rnx
#Synced-CL 49598926

diff --git a/Engine/Plugins/Runtime/Windows/WinDualShock/Source/WinDualShock/WinDualShock.Build.cs b/Engine/Plugins/Runtime/Windows/WinDualShock/Source/WinDualShock/WinDualShock.Build.cs
--- a/Engine/Plugins/Runtime/Windows/WinDualShock/Source/WinDualShock/WinDualShock.Build.cs
+++ b/Engine/Plugins/Runtime/Windows/WinDualShock/Source/WinDualShock/WinDualShock.Build.cs
@@ -1,6 +1,7 @@
 // Copyright Epic Games, Inc. All Rights Reserved.
 
 using System;
+using System.Reflection;
 
 namespace UnrealBuildTool.Rules
 {
@@ -29,10 +30,14 @@ namespace UnrealBuildTool.Rules
 			string PlatformName = "";
 			if (LibScePadType != null && Target.Architecture != UnrealArch.Arm64)
 			{
-				bHasSupport = (bool)LibScePadType.GetMethod("GetPadLibLocation").Invoke(null, new object[] { EngineDirectory, Target.WindowsPlatform.Compiler, null });
+				MethodInfo GetPadLibLocationMethod = LibScePadType.GetMethod("GetPadLibLocation");
+				object[] GetPadLibLocationMethodParams = new object[] { EngineDirectory, Target.WindowsPlatform.Compiler, null };
+				bHasSupport = (bool)GetPadLibLocationMethod.Invoke(null, GetPadLibLocationMethodParams);
 				PlatformName = (string)LibScePadType.GetMethod("GetPlatformName").Invoke(null, new object[] { EngineDirectory, Target.WindowsPlatform.Compiler });
 				if (bHasSupport)
 				{
+					ExtraRootPath = ("LibScePad", (string)GetPadLibLocationMethodParams[2]);
+
 					AddEngineThirdPartyPrivateStaticDependencies(Target,
 						"LibScePad",
 						"DX11Audio",
