[Win Dual Shock] Only create the audio endpoint submix for the WinDualShock plugin if it is a preferred input device API. In the case where users may have switched their preferred device API, we don't want to bother creating this submix endpoint if we dont need to.

We can accomplish this simply by checking the preferred device API name on module startup, instead of always creating it by default via the constructor.

#rnx
#Synced-CL 52444545

diff --git a/Engine/Plugins/Runtime/Windows/WinDualShock/Source/WinDualShock/Private/WinDualShock.cpp b/Engine/Plugins/Runtime/Windows/WinDualShock/Source/WinDualShock/Private/WinDualShock.cpp
--- a/Engine/Plugins/Runtime/Windows/WinDualShock/Source/WinDualShock/Private/WinDualShock.cpp
+++ b/Engine/Plugins/Runtime/Windows/WinDualShock/Source/WinDualShock/Private/WinDualShock.cpp
@@ -660,6 +660,10 @@ class FExternalDualShockEndpointFactory : public IAudioEndpointFactory
 public:
 	FExternalDualShockEndpointFactory(TMap<FDeviceKey, TSharedRef<FWinDualShockAudioDeviceImpl>>& InEarlyDeviceMap)
 		: EarlyDeviceMap(InEarlyDeviceMap)
+	{
+	}
+
+	void Register()
 	{
 		IAudioEndpointFactory::RegisterEndpointType(this);
 		bIsImplemented = true;
@@ -756,13 +760,21 @@ public:
 	{
 	}
 
-	virtual ~FWinDualShockPlugin()
+	virtual ~FWinDualShockPlugin() override
 	{
 	}
 
 	virtual void StartupModule() override
 	{
 		IInputDeviceModule::StartupModule();
+
+		// We only need to register the audio submix endpoints if this input API is preferred.
+		// Otherwise, we may be creating endpoints which are not useable. 
+		if (FGenericPlatformMisc::IsPreferredInputDevice(GetPreferredDeviceAPIString()))
+		{
+			PadSpeakerEndpoint.Register();
+			VibrationEndpoint.Register();	
+		}
 	}
 #endif
 
